HTML5 Experimental Networking plugin.

  -  uses libwebsockets.org
  -  currently enabled for HTML5 and Win64,
  -  Allows native and html5 clients on the same server.

diff --git a/Engine/Plugins/Experimental/HTML5Networking/HTML5Networking.uplugin b/Engine/Plugins/Experimental/HTML5Networking/HTML5Networking.uplugin
new file mode 100644
--- /dev/null
+++ b/Engine/Plugins/Experimental/HTML5Networking/HTML5Networking.uplugin
@@ -0,0 +1,22 @@
+{
+	"FileVersion" : 3,
+
+	"FriendlyName" : "Experimental HTML5 Networking Plugin",
+	"Version" : 1,
+	"VersionName" : "1.0",
+	"CreatedBy" : "Epic Games, Inc.",
+	"CreatedByURL" : "http://epicgames.com",
+	"EngineVersion" : "4.8.0",
+	"Description" : "HTML5 Networking",
+	"Category" : "Networking",
+	"EnabledByDefault" : false,
+
+	"Modules" :
+	[
+		{
+			"Name" : "HTML5Networking",
+			"Type" : "Runtime",
+			"WhitelistPlatforms" : [ "Win64", "HTML5" ]
+		}
+	]
+}
diff --git a/Engine/Plugins/Experimental/HTML5Networking/README.md b/Engine/Plugins/Experimental/HTML5Networking/README.md
new file mode 100644
--- /dev/null
+++ b/Engine/Plugins/Experimental/HTML5Networking/README.md
@@ -0,0 +1,42 @@
+### Unreal Engine 4 Networking over Websockets Plugin 
+
+- Provides websocket transport layer for unreal engine 4. 
+- Uses [libwebsockets](http://libwebsocket.org) for the server side and client side for non HTML5 clients.  
+- HTML5 clients use emscripten's sockets abstraction. 
+
+#### How to
+
+- Clone this in Engine\Plugins\Experimental directory. 
+- Add the following section in `BaseEngine.ini`
+```
+[/Script/HTML5Networking.WebSocketNetDriver]
+AllowPeerConnections=False
+AllowPeerVoice=False
+ConnectionTimeout=6000.0
+InitialConnectTimeout=6000.0
+AckTimeout=10.0
+KeepAliveTime=20.2
+MaxClientRate=15000
+MaxInternetClientRate=10000
+RelevantTimeout=5.0
+SpawnPrioritySeconds=1.0
+ServerTravelPause=4.0
+NetServerMaxTickRate=30
+LanServerMaxTickRate=35
+WebSocketPort=8889
+NetConnectionClassName="/Script/HTML5Networking.WebSocketConnection"
+MaxPortCountToTry=512
+```
+In section [/Script/Engine.Engine] disable comment out NetDriverDefinitions and add 
+```
+NetDriverDefinitions=(DefName="GameNetDriver",DriverClassName="/Script/HTML5Networking.WebSocketNetDriver",DriverClassNameFallback="/Script/HTML5Networking.IpNetDriver")
+```
+To enable this Net Driver. 
+
+Build! and follow existing Unreal Networking documentation to setup servers/clients.
+
+#### Issues/Todo 
+
+Disconnect events on client or server side are not handled properly yet 
+
+Copyright 2015 Epic Games. 
diff --git a/Engine/Plugins/Experimental/HTML5Networking/Source/HTML5Networking/Classes/WebSocketConnection.h b/Engine/Plugins/Experimental/HTML5Networking/Source/HTML5Networking/Classes/WebSocketConnection.h
new file mode 100644
--- /dev/null
+++ b/Engine/Plugins/Experimental/HTML5Networking/Source/HTML5Networking/Classes/WebSocketConnection.h
@@ -0,0 +1,29 @@
+// Copyright 1998-2015 Epic Games, Inc. All Rights Reserved.
+
+#pragma once
+
+#include "WebSocketConnection.generated.h"
+
+
+UCLASS(transient, config = Engine)
+class HTML5NETWORKING_API UWebSocketConnection : public UNetConnection
+{
+	GENERATED_UCLASS_BODY()
+
+	class FWebSocket*			WebSocket; 
+
+	// Begin NetConnection Interface
+	virtual void InitBase(UNetDriver* InDriver, class FSocket* InSocket, const FURL& InURL, EConnectionState InState, int32 InMaxPacket = 0, int32 InPacketOverhead = 0) override;
+	virtual void InitRemoteConnection(UNetDriver* InDriver, class FSocket* InSocket, const FURL& InURL, const class FInternetAddr& InRemoteAddr, EConnectionState InState, int32 InMaxPacket = 0, int32 InPacketOverhead = 0) override;
+	virtual void InitLocalConnection(UNetDriver* InDriver, class FSocket* InSocket, const FURL& InURL, EConnectionState InState, int32 InMaxPacket = 0, int32 InPacketOverhead = 0) override;
+	virtual void LowLevelSend(void* Data, int32 Count) override;
+	FString LowLevelGetRemoteAddress(bool bAppendPort = false) override;
+	FString LowLevelDescribe() override;
+	virtual void Tick();
+	virtual void FinishDestroy();
+	// End NetConnection Interface
+
+
+	void SetWebSocket(FWebSocket* InWebSocket);
+	FWebSocket* GetWebSocket();
+};
diff --git a/Engine/Plugins/Experimental/HTML5Networking/Source/HTML5Networking/Classes/WebSocketNetDriver.h b/Engine/Plugins/Experimental/HTML5Networking/Source/HTML5Networking/Classes/WebSocketNetDriver.h
new file mode 100644
--- /dev/null
+++ b/Engine/Plugins/Experimental/HTML5Networking/Source/HTML5Networking/Classes/WebSocketNetDriver.h
@@ -0,0 +1,67 @@
+// Copyright 1998-2015 Epic Games, Inc. All Rights Reserved.
+//
+// websocket based implementation of the net driver
+//
+
+#pragma once
+#include "WebSocketNetDriver.generated.h"
+
+UCLASS(transient, config = Engine)
+class HTML5NETWORKING_API UWebSocketNetDriver : public UNetDriver
+{
+	GENERATED_UCLASS_BODY()
+
+	/** Websocket server port*/
+	UPROPERTY(Config)
+	int32 WebSocketPort;
+
+	// Begin UNetDriver interface.
+	virtual bool IsAvailable() const override;
+	virtual bool InitBase(bool bInitAsClient, FNetworkNotify* InNotify, const FURL& URL, bool bReuseAddressAndPort, FString& Error) override;
+	virtual bool InitConnect(FNetworkNotify* InNotify, const FURL& ConnectURL, FString& Error) override;
+	virtual bool InitListen(FNetworkNotify* InNotify, FURL& LocalURL, bool bReuseAddressAndPort, FString& Error) override;
+	virtual void ProcessRemoteFunction(class AActor* Actor, class UFunction* Function, void* Parameters, struct FOutParmRec* OutParms, struct FFrame* Stack, class UObject* SubObject = NULL) override;
+	virtual void TickDispatch(float DeltaTime) override;
+	virtual FString LowLevelGetNetworkNumber() override;
+	virtual void LowLevelDestroy() override;
+	virtual bool IsNetResourceValid(void) override;
+
+	// stub implementation because for websockets we don't use any underlying socket sub system.
+	virtual class ISocketSubsystem* GetSocketSubsystem() override;
+	virtual FSocket * CreateSocket();
+
+	// End UNetDriver interface.
+
+	// Begin FExec Interface
+	virtual bool Exec(UWorld* InWorld, const TCHAR* Cmd, FOutputDevice& Ar = *GLog) override;
+	// End FExec Interface
+
+	/**
+	* Exec command handlers
+	*/
+	bool HandleSocketsCommand(const TCHAR* Cmd, FOutputDevice& Ar, UWorld* InWorld);
+
+	/** @return connection to server */
+	class UWebSocketConnection* GetServerConnection();
+
+	/************************************************************************/
+	/* FWebSocketServer														*/
+	/************************************************************************/
+
+	FWebSocketServer* WebSocketServer; 
+
+	/******************************************************************************/
+	/* Callback Function for New Connection from a client is accepted by this server   */
+	/******************************************************************************/
+
+	void OnWebSocketClientConnected(FWebSocket*); // to the server. 
+
+
+	/************************************************************************/
+	/* Callback Function for when this client Connects to the server				*/
+	/************************************************************************/
+
+	void OnWebSocketServerConnected(); // to the client. 
+
+
+};
diff --git a/Engine/Plugins/Experimental/HTML5Networking/Source/HTML5Networking/HTML5Networking.Build.cs b/Engine/Plugins/Experimental/HTML5Networking/Source/HTML5Networking/HTML5Networking.Build.cs
new file mode 100644
--- /dev/null
+++ b/Engine/Plugins/Experimental/HTML5Networking/Source/HTML5Networking/HTML5Networking.Build.cs
@@ -0,0 +1,22 @@
+// Copyright 1998-2015 Epic Games, Inc. All Rights Reserved.
+
+namespace UnrealBuildTool.Rules
+{
+    public class HTML5Networking  : ModuleRules
+	{
+		public HTML5Networking(TargetInfo Target)
+		{
+            PrivateDependencyModuleNames.AddRange(
+                new string[] { 
+                    "Core", 
+                    "CoreUObject",
+                    "Engine",
+                    "ImageCore",
+                    "Sockets",
+                    "OnlineSubSystemUtils",
+                    "WebSockets"
+                }
+            );
+		}
+	}
+}
diff --git a/Engine/Plugins/Experimental/HTML5Networking/Source/HTML5Networking/Private/HTML5NetworkingPCH.h b/Engine/Plugins/Experimental/HTML5Networking/Source/HTML5Networking/Private/HTML5NetworkingPCH.h
new file mode 100644
--- /dev/null
+++ b/Engine/Plugins/Experimental/HTML5Networking/Source/HTML5Networking/Private/HTML5NetworkingPCH.h
@@ -0,0 +1,21 @@
+// Copyright 1998-2015 Epic Games, Inc. All Rights Reserved.
+#pragma once
+
+#include "Core.h"
+#include "Engine.h"
+#include "Sockets.h"
+#include "SocketSubsystem.h"
+#include "ModuleManager.h"
+
+class FWebSocket; 
+class FWebSocketServer; 
+
+typedef struct libwebsocket_context WebSocketInternalContext; 
+typedef struct libwebsocket WebSocketInternal; 
+typedef struct libwebsocket_protocols WebSocketInternalProtocol; 
+
+DECLARE_DELEGATE_TwoParams(FWebsocketPacketRecievedCallBack, void* /*Data*/, int32 /*Data Size*/);
+DECLARE_DELEGATE_OneParam(FWebsocketClientConnectedCallBack, FWebSocket* /*Socket*/);
+DECLARE_DELEGATE(FWebsocketInfoCallBack); 
+
+DECLARE_LOG_CATEGORY_EXTERN(LogHTML5Networking, Warning, All);
\ No newline at end of file
diff --git a/Engine/Plugins/Experimental/HTML5Networking/Source/HTML5Networking/Private/HTML5NetworkingPlugin.cpp b/Engine/Plugins/Experimental/HTML5Networking/Source/HTML5Networking/Private/HTML5NetworkingPlugin.cpp
new file mode 100644
--- /dev/null
+++ b/Engine/Plugins/Experimental/HTML5Networking/Source/HTML5Networking/Private/HTML5NetworkingPlugin.cpp
@@ -0,0 +1,26 @@
+// Copyright 1998-2015 Epic Games, Inc. All Rights Reserved.
+
+#include "HTML5NetworkingPCH.h"
+
+
+class FHTML5NetworkingPlugin : public IModuleInterface
+{
+	/** IModuleInterface implementation */
+	virtual void StartupModule() override;
+	virtual void ShutdownModule() override;
+};
+
+IMPLEMENT_MODULE(FHTML5NetworkingPlugin, HTML5Networking)
+
+void FHTML5NetworkingPlugin::StartupModule()
+{
+}
+
+
+void FHTML5NetworkingPlugin::ShutdownModule()
+{
+}
+
+DEFINE_LOG_CATEGORY(LogHTML5Networking);
+
+
diff --git a/Engine/Plugins/Experimental/HTML5Networking/Source/HTML5Networking/Private/WebSocket.cpp b/Engine/Plugins/Experimental/HTML5Networking/Source/HTML5Networking/Private/WebSocket.cpp
new file mode 100644
--- /dev/null
+++ b/Engine/Plugins/Experimental/HTML5Networking/Source/HTML5Networking/Private/WebSocket.cpp
@@ -0,0 +1,410 @@
+// Copyright 1998-2015 Epic Games, Inc. All Rights Reserved.
+#include "HTML5NetworkingPCH.h"
+#include "WebSocket.h"
+
+#if PLATFORM_HTML5
+#include <errno.h>
+#include <sys/types.h>
+#include <sys/socket.h>
+#include <netinet/in.h>
+#include <arpa/inet.h>
+#include <stdio.h>
+#include <stdlib.h>
+#include <string.h>
+#include <unistd.h>
+#include <fcntl.h>
+#include <sys/ioctl.h>
+#include <assert.h>
+#include <emscripten.h>
+#endif
+
+#if PLATFORM_WINDOWS
+#include "AllowWindowsPlatformTypes.h"
+#endif
+
+
+#if !PLATFORM_HTML5
+#include "libwebsockets.h"
+#include "private-libwebsockets.h"
+#endif 
+
+#if PLATFORM_WINDOWS
+#include "HideWindowsPlatformTypes.h"
+#endif
+
+
+#if !PLATFORM_HTML5
+uint8 PREPADDING[LWS_SEND_BUFFER_PRE_PADDING];
+uint8 POSTPADDING[LWS_SEND_BUFFER_POST_PADDING];
+#endif 
+
+
+#if !PLATFORM_HTML5
+static void libwebsocket_debugLogS(int level, const char *line)
+{
+	UE_LOG(LogHTML5Networking, Log, TEXT("client: %s"), ANSI_TO_TCHAR(line));
+}
+#endif 
+
+FWebSocket::FWebSocket(
+		const FInternetAddr& ServerAddress
+)
+:IsServerSide(false)
+{
+
+#if !PLATFORM_HTML5_BROWSER
+
+#if !UE_BUILD_SHIPPING
+	lws_set_log_level(LLL_ERR | LLL_WARN | LLL_NOTICE | LLL_DEBUG | LLL_INFO, libwebsocket_debugLogS);
+#endif 
+
+	Protocols = new libwebsocket_protocols[3];
+	FMemory::Memzero(Protocols, sizeof(libwebsocket_protocols) * 3);
+
+	Protocols[0].name = "binary";
+	Protocols[0].callback = FWebSocket::unreal_networking_client;
+	Protocols[0].per_session_data_size = 0;
+	Protocols[0].rx_buffer_size = 10 * 1024 * 1024;
+
+	Protocols[1].name = nullptr;
+	Protocols[1].callback = nullptr;
+	Protocols[1].per_session_data_size = 0;
+
+	struct lws_context_creation_info Info;
+	memset(&Info, 0, sizeof Info);
+
+	Info.port = CONTEXT_PORT_NO_LISTEN;
+	Info.protocols = &Protocols[0];
+	Info.gid = -1;
+	Info.uid = -1;
+	Info.user = this;
+
+	Context = libwebsocket_create_context(&Info);
+
+	check(Context); 
+
+
+	Wsi = libwebsocket_client_connect_extended
+							(Context, 
+							TCHAR_TO_ANSI(*ServerAddress.ToString(false)), 
+							ServerAddress.GetPort(), 
+							false, "/", TCHAR_TO_ANSI(*ServerAddress.ToString(false)), TCHAR_TO_ANSI(*ServerAddress.ToString(false)), Protocols[1].name, -1,this);
+
+	check(Wsi);
+
+#endif 
+
+#if PLATFORM_HTML5_BROWSER
+
+	struct sockaddr_in addr;
+	int res;
+
+	SockFd = socket(PF_INET, SOCK_STREAM, IPPROTO_TCP);
+	if (SockFd == -1) {
+		UE_LOG(LogHTML5Networking, Error, TEXT("Socket creationg failed "));
+	}
+	else
+	{
+		UE_LOG(LogHTML5Networking, Warning, TEXT(" Socked %d created "), SockFd);
+	}
+
+	fcntl(SockFd, F_SETFL, O_NONBLOCK);
+
+	memset(&addr, 0, sizeof(addr));
+	addr.sin_family = AF_INET;
+	addr.sin_port = htons(ServerAddress.GetPort());
+
+	if (inet_pton(AF_INET, TCHAR_TO_ANSI(*ServerAddress.ToString(false)), &addr.sin_addr) != 1) {
+		UE_LOG(LogHTML5Networking, Warning, TEXT("inet_pton failed "));
+		return; 
+	}
+
+	int Ret = connect(SockFd, (struct sockaddr *)&addr, sizeof(addr));
+	UE_LOG(LogHTML5Networking, Warning, TEXT(" Connect socket returned %d"), Ret);
+
+#endif 
+}
+
+FWebSocket::FWebSocket(WebSocketInternalContext* InContext, WebSocketInternal* InWsi)
+:	Context(InContext), 
+	Wsi(InWsi),
+	IsServerSide(true)
+{
+}
+
+
+bool FWebSocket::Send(uint8* Data, uint32 Size)
+{
+	TArray<uint8> Buffer;
+	// insert size. 
+
+#if !PLATFORM_HTML5
+	Buffer.Append((uint8*)&PREPADDING, sizeof(PREPADDING));
+#endif 
+
+	Buffer.Append((uint8*)&Size, sizeof (uint32));
+	Buffer.Append((uint8*)Data, Size);
+
+#if !PLATFORM_HTML5
+	Buffer.Append((uint8*)&POSTPADDING, sizeof(POSTPADDING));
+#endif 
+
+	OutgoingBuffer.Add(Buffer);
+
+	return true;
+}
+
+void FWebSocket::SetRecieveCallBack(FWebsocketPacketRecievedCallBack CallBack)
+{
+	RecievedCallBack = CallBack; 
+}
+
+FString FWebSocket::RemoteEndPoint()
+{
+#if !PLATFORM_HTML5
+	ANSICHAR Peer_Name[128];
+	ANSICHAR Peer_Ip[128];
+	libwebsockets_get_peer_addresses(Context, Wsi, libwebsocket_get_socket_fd(Wsi), Peer_Name, sizeof Peer_Name, Peer_Ip, sizeof Peer_Ip);
+	return FString(Peer_Name);
+#endif
+
+#if PLATFORM_HTML5
+	return FString(TEXT("TODO:REMOTEENDPOINT"));
+#endif
+}
+
+
+FString FWebSocket::LocalEndPoint()
+{
+#if !PLATFORM_HTML5
+	return FString(ANSI_TO_TCHAR(libwebsocket_canonical_hostname(Context)));
+#endif 
+
+#if PLATFORM_HTML5
+	return FString(TEXT("TODO:LOCALENDPOINT"));
+#endif
+
+}
+
+void FWebSocket::Tick()
+{
+
+#if !PLATFORM_HTML5
+	{
+		libwebsocket_service(Context, 0);
+		if (!IsServerSide)
+			libwebsocket_callback_on_writable_all_protocol(&Protocols[0]);
+	}
+#endif 
+
+#if PLATFORM_HTML5
+
+	fd_set Fdr;
+	fd_set Fdw;
+	int Res;
+
+	// make sure that server.fd is ready to read / write
+	FD_ZERO(&Fdr);
+	FD_ZERO(&Fdw);
+	FD_SET(SockFd, &Fdr);
+	FD_SET(SockFd, &Fdw);
+	Res = select(64, &Fdr, &Fdw, NULL, NULL);
+
+	if (Res == -1) {
+		UE_LOG(LogHTML5Networking, Warning, TEXT("Select Failed!"));
+		return;
+	}
+	
+	if (FD_ISSET(SockFd, &Fdr)) {
+		// we can read! 
+		OnRawRecieve(NULL, NULL);
+	}
+
+	if (FD_ISSET(SockFd, &Fdw)) {
+		// we can write
+		OnRawWebSocketWritable(NULL);
+	}
+#endif 
+}
+
+void FWebSocket::SetConnectedCallBack(FWebsocketInfoCallBack CallBack)
+{
+	ConnectedCallBack = CallBack; 
+}
+
+void FWebSocket::SetErrorCallBack(FWebsocketInfoCallBack CallBack)
+{
+	ErrorCallBack = CallBack; 
+}
+
+void FWebSocket::OnRawRecieve(void* Data, uint32 Size)
+{
+#if !PLATFORM_HTML5
+	RecievedBuffer.Append((uint8*)Data, Size);
+
+	while (RecievedBuffer.Num())
+	{
+		uint32 BytesToBeRead = *(uint32*)RecievedBuffer.GetData();
+		if (BytesToBeRead <= ((uint32)RecievedBuffer.Num() - sizeof(uint32)))
+		{
+			RecievedCallBack.Execute((void*)((uint8*)RecievedBuffer.GetData() + sizeof(uint32)), BytesToBeRead);
+			RecievedBuffer.RemoveAt(0, sizeof(uint32) + BytesToBeRead );
+		}
+		else
+		{
+			break;
+		}
+	}
+#endif
+
+#if PLATFORM_HTML5
+
+	uint8 Buffer[1024]; // should be at MAX PACKET SIZE. 
+	int Result = recv(SockFd, Buffer, sizeof(uint32), 0);
+
+	uint32 DataToBeRead = 0;*(uint32*)Buffer;
+
+	if (Result == -1) 
+	{
+		UE_LOG(LogHTML5Networking, Log, TEXT("Read message size failed!"));
+		this->ErrorCallBack.ExecuteIfBound(); 
+	}
+	else
+	{
+		DataToBeRead = *(uint32*)Buffer;
+		UE_LOG(LogHTML5Networking, Log, TEXT("Read 4 bytes showing the size"), DataToBeRead);
+	}
+
+	check(Result == sizeof(uint32)); 
+
+	// read rest of the data. 
+	Result = recv(SockFd, Buffer, DataToBeRead, 0);
+
+	if(Result < 0 )
+	{
+		UE_LOG(LogHTML5Networking, Log, TEXT("Read message failed!"));
+		this->ErrorCallBack.ExecuteIfBound();
+	}
+	else
+	{
+		UE_LOG(LogHTML5Networking, Log, TEXT("Read %d bytes and Executing."), DataToBeRead);
+		check(DataToBeRead == Result);
+		RecievedCallBack.Execute(Buffer, DataToBeRead);
+	}
+
+#endif 
+
+}
+
+void FWebSocket::OnRawWebSocketWritable(WebSocketInternal* wsi)
+{
+	if (OutgoingBuffer.Num() == 0)
+		return;
+
+	TArray <uint8>& Packet = OutgoingBuffer[0];
+
+#if !PLATFORM_HTML5_BROWSER
+
+	uint32 TotalDataSize = Packet.Num() - LWS_SEND_BUFFER_PRE_PADDING - LWS_SEND_BUFFER_POST_PADDING;
+
+	int Sent = libwebsocket_write(Wsi, Packet.GetData() + LWS_SEND_BUFFER_PRE_PADDING, TotalDataSize, (libwebsocket_write_protocol)LWS_WRITE_BINARY);
+
+	check(Sent == TotalDataSize); // if this fires we need a slightly more complicated buffering mechanism. 
+	check(Wsi == wsi);
+
+#endif
+
+#if  PLATFORM_HTML5_BROWSER
+
+	// send actual data in one go. 
+	int Result = send(SockFd, (uint32*)Packet.GetData(),Packet.Num(), 0);
+
+	if (Result == -1)
+	{
+		// we are caught with our pants down. fail. 
+		UE_LOG(LogHTML5Networking, Error, TEXT("Could not write %d bytes"), Packet.Num());
+		ErrorCallBack.ExecuteIfBound(); 
+	}
+	else
+	{
+		check(Result == Packet.Num());
+		UE_LOG(LogHTML5Networking, Log, TEXT("Wrote %d bytes"), Packet.Num());
+	}
+	
+#endif 
+
+	// this is very inefficient we need a constant size circular buffer to efficiently not do unnecessary allocations/deallocations. 
+	OutgoingBuffer.RemoveAt(0);
+
+}
+
+FWebSocket::~FWebSocket()
+{
+#if !PLATFORM_HTML5
+	if ( !IsServerSide)
+	{
+		libwebsocket_context_destroy(Context);
+		Context = NULL;
+		delete Protocols;
+		Protocols = NULL;
+	}
+#endif 
+
+#if PLATFORM_HTML5
+	close(SockFd);
+#endif 
+
+}
+
+#if !PLATFORM_HTML5
+int FWebSocket::unreal_networking_client(
+		struct libwebsocket_context *Context, 
+		struct libwebsocket *Wsi, 
+		enum libwebsocket_callback_reasons Reason, 
+		void *User, 
+		void *In, 
+		size_t Len)
+{
+	FWebSocket* Socket = (FWebSocket*)libwebsocket_context_user(Context);;
+	switch (Reason)
+	{
+	case LWS_CALLBACK_CLIENT_ESTABLISHED:
+		{
+			Socket->ConnectedCallBack.ExecuteIfBound();
+			libwebsocket_set_timeout(Wsi, NO_PENDING_TIMEOUT, 0);
+			check(Socket->Wsi == Wsi);
+		}
+		break;
+	case LWS_CALLBACK_CLIENT_CONNECTION_ERROR:
+		{
+			Socket->ErrorCallBack.ExecuteIfBound();
+			return -1;
+		}
+		break;
+	case LWS_CALLBACK_CLIENT_RECEIVE:
+		{
+			// push it on the socket. 
+			Socket->OnRawRecieve(In, (uint32)Len); 
+			check(Socket->Wsi == Wsi);
+			libwebsocket_set_timeout(Wsi, NO_PENDING_TIMEOUT, 0);
+			break;
+		}
+	case LWS_CALLBACK_CLIENT_WRITEABLE:
+		{
+			check(Socket->Wsi == Wsi);
+			Socket->OnRawWebSocketWritable(Wsi); 
+			libwebsocket_callback_on_writable(Context, Wsi);
+			libwebsocket_set_timeout(Wsi, NO_PENDING_TIMEOUT, 0);
+			break; 
+		}
+	case LWS_CALLBACK_CLOSED:
+		{
+			Socket->ErrorCallBack.ExecuteIfBound();
+			return -1;
+		}
+	}
+
+	return 0; 
+}
+#endif 
+
diff --git a/Engine/Plugins/Experimental/HTML5Networking/Source/HTML5Networking/Private/WebSocket.h b/Engine/Plugins/Experimental/HTML5Networking/Source/HTML5Networking/Private/WebSocket.h
new file mode 100644
--- /dev/null
+++ b/Engine/Plugins/Experimental/HTML5Networking/Source/HTML5Networking/Private/WebSocket.h
@@ -0,0 +1,77 @@
+// Copyright 1998-2015 Epic Games, Inc. All Rights Reserved.
+//
+// libwebsocket client wrapper. 
+//
+#pragma  once
+#include "HTML5NetworkingPCH.h"
+
+class FWebSocket
+{
+
+public: 
+
+	// Initialize as client side socket. 
+	FWebSocket(const FInternetAddr& ServerAddress);
+
+	// Initialize as server side socket. 
+	FWebSocket(WebSocketInternalContext* InContext, WebSocketInternal* Wsi);
+
+	// clean up. 
+	~FWebSocket();
+
+	/************************************************************************/
+	/* Set various callbacks for Socket Events                              */           
+	/************************************************************************/
+	void SetConnectedCallBack(FWebsocketInfoCallBack CallBack);
+	void SetErrorCallBack(FWebsocketInfoCallBack CallBack);
+	void SetRecieveCallBack(FWebsocketPacketRecievedCallBack CallBack);
+
+	/** Send raw data to remote end point. */ 
+	bool Send(uint8* Data, uint32 Size);
+
+	/** service libwebsocket.			   */ 
+	void Tick();
+
+	/** Helper functions to describe end points. */
+	FString RemoteEndPoint();
+	FString LocalEndPoint(); 
+
+private:
+
+	void OnRawRecieve(void* Data, uint32 Size);
+	void OnRawWebSocketWritable(WebSocketInternal* wsi);
+
+	/************************************************************************/
+	/*	Various Socket callbacks											*/                                                                 
+	/************************************************************************/ 
+	FWebsocketPacketRecievedCallBack  RecievedCallBack; 
+	FWebsocketInfoCallBack ConnectedCallBack;
+	FWebsocketInfoCallBack ErrorCallBack;
+
+	/**  Recv and Send Buffers, serviced during the Tick */
+	TArray<uint8> RecievedBuffer;
+	TArray<TArray<uint8>> OutgoingBuffer;
+
+	/** libwebsocket internal context*/
+	WebSocketInternalContext* Context;
+
+	/** libwebsocket web socket */
+	WebSocketInternal* Wsi;
+
+	/** libwebsocket Protocols that can be serviced by this implemenation*/
+	WebSocketInternalProtocol* Protocols;
+
+	/** Server side socket or client side*/
+	bool IsServerSide; 
+
+#if !PLATFORM_HTML5
+	/* libwebsocket service functions */ 
+	static int unreal_networking_server(struct libwebsocket_context *, struct libwebsocket *wsi, enum libwebsocket_callback_reasons reason, void *user, void *in, size_t len);
+	static int unreal_networking_client(struct libwebsocket_context *, struct libwebsocket *wsi, enum libwebsocket_callback_reasons reason, void *user, void *in, size_t len);
+#endif 
+
+	friend class FWebSocketServer;
+	int SockFd;
+};
+
+
diff --git a/Engine/Plugins/Experimental/HTML5Networking/Source/HTML5Networking/Private/WebSocketNetDriver.cpp b/Engine/Plugins/Experimental/HTML5Networking/Source/HTML5Networking/Private/WebSocketNetDriver.cpp
new file mode 100644
--- /dev/null
+++ b/Engine/Plugins/Experimental/HTML5Networking/Source/HTML5Networking/Private/WebSocketNetDriver.cpp
@@ -0,0 +1,277 @@
+// Copyright 1998-2015 Epic Games, Inc. All Rights Reserved.
+
+/*=============================================================================
+Unreal Websocket network driver.
+=============================================================================*/
+
+#include "HTML5NetworkingPCH.h"
+#include "Engine/Channel.h"
+
+#include "IPAddress.h"
+#include "Sockets.h"
+
+#include "WebSocketConnection.h"
+#include "WebSocketNetDriver.h"
+#include "WebSocketServer.h"
+#include "WebSocket.h"
+
+
+/*-----------------------------------------------------------------------------
+Declarations.
+-----------------------------------------------------------------------------*/
+
+/** Size of the network recv buffer */
+#define NETWORK_MAX_PACKET (576)
+
+UWebSocketNetDriver::UWebSocketNetDriver(const FObjectInitializer& ObjectInitializer)
+	: Super(ObjectInitializer)
+{
+}
+
+bool UWebSocketNetDriver::IsAvailable() const
+{
+	// Websocket driver always valid for now
+	return true;
+}
+
+
+
+ISocketSubsystem* UWebSocketNetDriver::GetSocketSubsystem()
+{
+	return ISocketSubsystem::Get();
+}
+
+FSocket * UWebSocketNetDriver::CreateSocket()
+{
+	return NULL; 
+}
+
+bool UWebSocketNetDriver::InitBase(bool bInitAsClient, FNetworkNotify* InNotify, const FURL& URL, bool bReuseAddressAndPort, FString& Error)
+{
+	if (!Super::InitBase(bInitAsClient, InNotify, URL, bReuseAddressAndPort, Error))
+	{
+		return false;
+	}
+
+	return true;
+}
+
+bool UWebSocketNetDriver::InitConnect(FNetworkNotify* InNotify, const FURL& ConnectURL, FString& Error)
+{
+	if (!InitBase(true, InNotify, ConnectURL, false, Error))
+	{
+		return false;
+	}
+
+	// Create new connection.
+	ServerConnection = NewObject<UWebSocketConnection>(NetConnectionClass);
+
+	TSharedRef<FInternetAddr> InternetAddr = GetSocketSubsystem()->CreateInternetAddr();
+	bool Ok;
+	InternetAddr->SetIp(*ConnectURL.Host, Ok);
+	InternetAddr->SetPort(WebSocketPort);
+
+	FWebSocket* WebSocket = new FWebSocket(*InternetAddr);
+	UWebSocketConnection* Connection = (UWebSocketConnection*)ServerConnection;
+	Connection->SetWebSocket(WebSocket);
+
+	FWebsocketPacketRecievedCallBack CallBack;
+	CallBack.BindUObject(Connection, &UWebSocketConnection::ReceivedRawPacket);
+	WebSocket->SetRecieveCallBack(CallBack);
+
+	FWebsocketInfoCallBack  ConnectedCallBack;
+	ConnectedCallBack.BindUObject(this, &UWebSocketNetDriver::OnWebSocketServerConnected);
+	WebSocket->SetConnectedCallBack(ConnectedCallBack);
+
+	ServerConnection->InitLocalConnection(this, NULL, ConnectURL, USOCK_Pending);
+
+	// Create channel zero.
+	GetServerConnection()->CreateChannel(CHTYPE_Control, 1, 0);
+
+	return true;
+}
+
+bool UWebSocketNetDriver::InitListen(FNetworkNotify* InNotify, FURL& LocalURL, bool bReuseAddressAndPort, FString& Error)
+{
+	if (!InitBase(false, InNotify, LocalURL, bReuseAddressAndPort, Error))
+	{
+		return false;
+	}
+
+	WebSocketServer = new class FWebSocketServer();
+
+	FWebsocketClientConnectedCallBack CallBack;
+	CallBack.BindUObject(this, &UWebSocketNetDriver::OnWebSocketClientConnected);
+
+	if(!WebSocketServer->Init(WebSocketPort, CallBack))
+		return false; 
+
+	WebSocketServer->Tick();
+	LocalURL.Port = WebSocketPort;
+	UE_LOG(LogHTML5Networking, Log, TEXT("%s WebSocketNetDriver listening on port %i"), *GetDescription(), LocalURL.Port);
+
+	// server has no server connection. 
+	ServerConnection = NULL; 
+	return true;
+}
+
+void UWebSocketNetDriver::TickDispatch(float DeltaTime)
+{
+	Super::TickDispatch(DeltaTime);
+
+	if (WebSocketServer)
+		WebSocketServer->Tick();
+}
+
+void UWebSocketNetDriver::ProcessRemoteFunction(class AActor* Actor, UFunction* Function, void* Parameters, FOutParmRec* OutParms, FFrame* Stack, class UObject* SubObject)
+{
+	bool bIsServer = IsServer();
+
+	UNetConnection* Connection = NULL;
+	if (bIsServer)
+	{
+		if ((Function->FunctionFlags & FUNC_NetMulticast))
+		{
+			// Multicast functions go to every client
+			TArray<UNetConnection*> UniqueRealConnections;
+			for (int32 i = 0; i<ClientConnections.Num(); ++i)
+			{
+				Connection = ClientConnections[i];
+				if (Connection)
+				{
+					// Do relevancy check if unreliable.
+					// Reliables will always go out. This is odd behavior. On one hand we wish to garuntee "reliables always get there". On the other
+					// hand, replicating a reliable to something on the other side of the map that is non relevant seems weird. 
+					//
+					// Multicast reliables should probably never be used in gameplay code for actors that have relevancy checks. If they are, the 
+					// rpc will go through and the channel will be closed soon after due to relevancy failing.
+
+					bool IsRelevant = true;
+					if ((Function->FunctionFlags & FUNC_NetReliable) == 0)
+					{
+						if (Connection->ViewTarget)
+						{
+							FNetViewer Viewer(Connection, 0.f);
+							IsRelevant = Actor->IsNetRelevantFor(Viewer.InViewer, Viewer.ViewTarget, Viewer.ViewLocation);
+						}
+						else
+						{
+							// No viewer for this connection(?), just let it go through.
+							UE_LOG(LogHTML5Networking, Log, TEXT("Multicast function %s called on actor %s when a connection has no Viewer"), *Function->GetName(), *Actor->GetName());
+						}
+					}
+
+					if (IsRelevant)
+					{
+						if (Connection->GetUChildConnection() != NULL)
+						{
+							Connection = ((UChildConnection*)Connection)->Parent;
+						}
+
+						InternalProcessRemoteFunction(Actor, SubObject, Connection, Function, Parameters, OutParms, Stack, bIsServer);
+					}
+				}
+			}
+
+			// Return here so we don't call InternalProcessRemoteFunction again at the bottom of this function
+			return;
+		}
+	}
+
+	// Send function data to remote.
+	Connection = Actor->GetNetConnection();
+	if (Connection)
+	{
+		InternalProcessRemoteFunction(Actor, SubObject, Connection, Function, Parameters, OutParms, Stack, bIsServer);
+	}
+}
+
+FString UWebSocketNetDriver::LowLevelGetNetworkNumber()
+{
+	return WebSocketServer->Info(); 
+}
+
+void UWebSocketNetDriver::LowLevelDestroy()
+{
+	Super::LowLevelDestroy();
+	delete WebSocketServer; 
+}
+
+bool UWebSocketNetDriver::HandleSocketsCommand(const TCHAR* Cmd, FOutputDevice& Ar, UWorld* InWorld)
+{
+	Ar.Logf(TEXT(""));
+	if (WebSocketServer != NULL)
+	{
+		Ar.Logf(TEXT("Running WebSocket Server %s"), *WebSocketServer->Info());
+	}
+	else
+	{
+		check(GetServerConnection());
+		Ar.Logf(TEXT("WebSocket client's EndPoint %s"), *(GetServerConnection()->WebSocket->RemoteEndPoint()));
+	}
+	return UNetDriver::Exec(InWorld, TEXT("SOCKETS"), Ar);
+}
+
+bool UWebSocketNetDriver::Exec(UWorld* InWorld, const TCHAR* Cmd, FOutputDevice& Ar)
+{
+	if (FParse::Command(&Cmd, TEXT("SOCKETS")))
+	{
+		return HandleSocketsCommand(Cmd, Ar, InWorld);
+	}
+	return UNetDriver::Exec(InWorld, Cmd, Ar);
+}
+
+UWebSocketConnection* UWebSocketNetDriver::GetServerConnection()
+{
+	return (UWebSocketConnection*)ServerConnection;
+}
+
+void UWebSocketNetDriver::OnWebSocketClientConnected(FWebSocket* ClientWebSocket)
+{
+	// Determine if allowing for client/server connections
+	const bool bAcceptingConnection = Notify->NotifyAcceptingConnection() == EAcceptConnection::Accept;
+	if (bAcceptingConnection)
+	{
+
+		UWebSocketConnection* Connection = NewObject<UWebSocketConnection>(NetConnectionClass);
+		check(Connection); 
+
+		TSharedRef<FInternetAddr> InternetAddr = GetSocketSubsystem()->CreateInternetAddr();
+		bool Ok;
+
+		InternetAddr->SetIp(*(ClientWebSocket->RemoteEndPoint()),Ok);
+		InternetAddr->SetPort(0);
+		Connection->SetWebSocket(ClientWebSocket);
+		Connection->InitRemoteConnection(this, NULL, FURL(), *InternetAddr, USOCK_Open);
+		
+		Notify->NotifyAcceptedConnection(Connection);
+
+		AddClientConnection(Connection);
+
+ 		FWebsocketPacketRecievedCallBack CallBack;
+ 		CallBack.BindUObject(Connection, &UWebSocketConnection::ReceivedRawPacket);
+ 		ClientWebSocket->SetRecieveCallBack(CallBack);
+
+		UE_LOG(LogHTML5Networking, Log, TEXT(" Websocket server running on %s Accepted Connection from %s "), *WebSocketServer->Info(),*ClientWebSocket->RemoteEndPoint());
+	}
+}
+
+bool UWebSocketNetDriver::IsNetResourceValid(void)
+{
+	if (	(WebSocketServer && !ServerConnection)//  Server 
+		||	(!WebSocketServer && ServerConnection) // client
+		)
+	{
+		return true; 
+	}
+		
+	return false; 
+}
+
+// Just logging, not yet attached to html5 clients.
+void UWebSocketNetDriver::OnWebSocketServerConnected()
+{
+	check(GetServerConnection());
+	UE_LOG(LogHTML5Networking, Log, TEXT(" %s Websocket Client connected to server %s "), *GetDescription(), *GetServerConnection()->WebSocket->RemoteEndPoint());
+}
+
diff --git a/Engine/Plugins/Experimental/HTML5Networking/Source/HTML5Networking/Private/WebSocketServer.cpp b/Engine/Plugins/Experimental/HTML5Networking/Source/HTML5Networking/Private/WebSocketServer.cpp
new file mode 100644
--- /dev/null
+++ b/Engine/Plugins/Experimental/HTML5Networking/Source/HTML5Networking/Private/WebSocketServer.cpp
@@ -0,0 +1,181 @@
+// Copyright 1998-2015 Epic Games, Inc. All Rights Reserved.
+
+#include "HTML5NetworkingPCH.h"
+#include "WebSocketServer.h"
+#include "WebSocket.h"
+
+#if PLATFORM_WINDOWS
+#include "AllowWindowsPlatformTypes.h"
+#endif
+
+#if !PLATFORM_HTML5
+#include "libwebsockets.h"
+#endif 
+
+#if PLATFORM_WINDOWS
+#include "HideWindowsPlatformTypes.h"
+#endif
+
+// a object of this type is associated by libwebsocket to every connected session. 
+struct PerSessionDataServer
+{
+	FWebSocket *Socket; // each session is actually a socket to a client
+};
+
+
+#if !PLATFORM_HTML5
+// real networking handler. 
+static int unreal_networking_server(
+	struct libwebsocket_context *,
+	struct libwebsocket *wsi,
+	enum libwebsocket_callback_reasons reason,
+	void *user,
+	void *in,
+	size_t
+	len
+);
+
+#if !UE_BUILD_SHIPPING
+	void libwebsocket_debugLog(int level, const char *line)
+	{ 
+		UE_LOG(LogHTML5Networking, Log, TEXT("websocket server: %s"), ANSI_TO_TCHAR(line));
+	}
+#endif 
+#endif 
+
+bool FWebSocketServer::Init(uint32 Port, FWebsocketClientConnectedCallBack CallBack)
+{
+#if !PLATFORM_HTML5
+	// setup log level.
+#if !UE_BUILD_SHIPPING
+	lws_set_log_level(LLL_ERR | LLL_WARN | LLL_NOTICE | LLL_DEBUG | LLL_INFO, libwebsocket_debugLog);
+#endif 
+
+	Protocols = new libwebsocket_protocols[3];
+	FMemory::Memzero(Protocols, sizeof(libwebsocket_protocols) * 3);
+
+	Protocols[0].name = "binary";
+	Protocols[0].callback = FWebSocket::unreal_networking_server;
+	Protocols[0].per_session_data_size = sizeof(PerSessionDataServer);
+	Protocols[0].rx_buffer_size = 10 * 1024 * 1024;
+
+	Protocols[1].name = nullptr;
+	Protocols[1].callback = nullptr;
+	Protocols[1].per_session_data_size = 0;
+
+	struct lws_context_creation_info Info;
+	memset(&Info, 0, sizeof(lws_context_creation_info));
+	// look up libwebsockets.h for details. 
+	Info.port = Port;
+	// we listen on all available interfaces. 
+	Info.iface = NULL;
+	Info.protocols = &Protocols[0];
+	// no extensions
+	Info.extensions = NULL;
+	Info.gid = -1;
+	Info.uid = -1;
+	Info.options = 0;
+	// tack on this object. 
+	Info.user = this;
+	Info.port = Port; 
+	Context = libwebsocket_create_context(&Info);
+
+	if (Context == NULL) 
+	{
+		return false; // couldn't create a server. 
+	}
+
+	ConnectedCallBack = CallBack; 
+
+#endif
+	return true; 
+}
+
+bool FWebSocketServer::Tick()
+{
+#if !PLATFORM_HTML5
+	{
+		libwebsocket_service(Context, 0);
+		libwebsocket_callback_on_writable_all_protocol(&Protocols[0]);
+	}
+#endif 
+	return true;
+}
+
+FWebSocketServer::FWebSocketServer()
+{}
+
+FWebSocketServer::~FWebSocketServer()
+{
+#if !PLATFORM_HTML5
+	if (Context)
+	{
+
+		libwebsocket_context_destroy(Context);
+		Context = NULL;
+	}
+
+	 delete Protocols; 
+	 Protocols = NULL; 
+#endif 		
+}
+
+FString FWebSocketServer::Info()
+{
+
+#if !PLATFORM_HTML5
+	return FString(ANSI_TO_TCHAR(libwebsocket_canonical_hostname(Context)));
+#else 
+	return FString(TEXT("NOT SUPPORTED"));
+#endif 
+
+}
+
+// callback. 
+#if !PLATFORM_HTML5
+int FWebSocket::unreal_networking_server
+	(
+		struct libwebsocket_context *Context, 
+		struct libwebsocket *Wsi, 
+		enum libwebsocket_callback_reasons Reason, 
+		void *User, 
+		void *In, 
+		size_t Len
+	)
+{
+	PerSessionDataServer* BufferInfo = (PerSessionDataServer*)User;
+	FWebSocketServer* Server = (FWebSocketServer*)libwebsocket_context_user(Context);
+
+	switch (Reason)
+	{
+		case LWS_CALLBACK_ESTABLISHED: 
+			{
+				BufferInfo->Socket = new FWebSocket(Context, Wsi);
+				Server->ConnectedCallBack.ExecuteIfBound(BufferInfo->Socket);
+				libwebsocket_set_timeout(Wsi, NO_PENDING_TIMEOUT, 0);
+			}
+			break;
+
+		case LWS_CALLBACK_RECEIVE:
+			{
+				BufferInfo->Socket->OnRawRecieve(In, Len);
+				libwebsocket_set_timeout(Wsi, NO_PENDING_TIMEOUT, 0);
+			}
+			break; 
+
+		case LWS_CALLBACK_SERVER_WRITEABLE: 
+			{
+				BufferInfo->Socket->OnRawWebSocketWritable(Wsi);
+				libwebsocket_set_timeout(Wsi, NO_PENDING_TIMEOUT, 0);
+			}
+			break; 
+		case LWS_CALLBACK_CLIENT_CONNECTION_ERROR:
+			{
+				BufferInfo->Socket->ErrorCallBack.ExecuteIfBound();
+			}
+			break;
+	}
+
+	return 0; 
+}
+#endif 
\ No newline at end of file
diff --git a/Engine/Plugins/Experimental/HTML5Networking/Source/HTML5Networking/Private/WebSocketServer.h b/Engine/Plugins/Experimental/HTML5Networking/Source/HTML5Networking/Private/WebSocketServer.h
new file mode 100644
--- /dev/null
+++ b/Engine/Plugins/Experimental/HTML5Networking/Source/HTML5Networking/Private/WebSocketServer.h
@@ -0,0 +1,38 @@
+// Copyright 1998-2015 Epic Games, Inc. All Rights Reserved.
+//
+// Read http://lucumr.pocoo.org/2012/9/24/websockets-101/ for a nice intro to web sockets. 
+// This uses https://libwebsockets.org/trac/libwebsockets 
+#pragma  once
+#include "HTML5NetworkingPCH.h"
+
+class FWebSocketServer
+{
+public:
+
+	FWebSocketServer(); 
+	~FWebSocketServer();
+
+	/** Create a web socket server*/
+	bool Init(uint32 Port, FWebsocketClientConnectedCallBack);
+
+	/** Service libwebsocket */
+	bool Tick();
+
+	/** Describe this libwebsocket server */
+	FString Info(); 
+
+private: 
+
+	/** Callback for a new websocket connection to the server */
+	FWebsocketClientConnectedCallBack  ConnectedCallBack;
+
+	/** Internal libwebsocket context */
+	WebSocketInternalContext* Context;
+
+	/** Protocols serviced by this implementation */
+	WebSocketInternalProtocol* Protocols;
+
+	friend class FWebSocket;
+};
+
+
diff --git a/Engine/Plugins/Experimental/HTML5Networking/Source/HTML5Networking/Private/WebsocketConnection.cpp b/Engine/Plugins/Experimental/HTML5Networking/Source/HTML5Networking/Private/WebsocketConnection.cpp
new file mode 100644
--- /dev/null
+++ b/Engine/Plugins/Experimental/HTML5Networking/Source/HTML5Networking/Private/WebsocketConnection.cpp
@@ -0,0 +1,117 @@
+// Copyright 1998-2015 Epic Games, Inc. All Rights Reserved.
+
+#include "HTML5NetworkingPCH.h"
+
+#include "IPAddress.h"
+#include "Sockets.h"
+#include "Net/NetworkProfiler.h"
+#include "Net/DataChannel.h"
+
+#include "WebSocketConnection.h"
+#include "WebSocketNetDriver.h"
+#include "WebSocket.h"
+/*-----------------------------------------------------------------------------
+Declarations.
+-----------------------------------------------------------------------------*/
+
+// Size of a UDP header.
+#define IP_HEADER_SIZE     (20)
+#define UDP_HEADER_SIZE    (IP_HEADER_SIZE+8)
+#define SLIP_HEADER_SIZE   (UDP_HEADER_SIZE+4)
+#define WINSOCK_MAX_PACKET (512)
+
+UWebSocketConnection::UWebSocketConnection(const FObjectInitializer& ObjectInitializer) :
+Super(ObjectInitializer),
+WebSocket(NULL)
+{
+}
+
+void UWebSocketConnection::InitBase(UNetDriver* InDriver, class FSocket* InSocket, const FURL& InURL, EConnectionState InState, int32 InMaxPacket, int32 InPacketOverhead)
+{
+	// Pass the call up the chain
+	Super::InitBase(InDriver, InSocket, InURL, InState,
+		// Use the default packet size/overhead unless overridden by a child class
+		InMaxPacket == 0 ? WINSOCK_MAX_PACKET : InMaxPacket,
+		InPacketOverhead == 0 ? SLIP_HEADER_SIZE : InPacketOverhead);
+
+}
+
+void UWebSocketConnection::InitLocalConnection(UNetDriver* InDriver, class FSocket* InSocket, const FURL& InURL, EConnectionState InState, int32 InMaxPacket, int32 InPacketOverhead)
+{
+	InitBase(InDriver, InSocket, InURL, InState,
+		// Use the default packet size/overhead unless overridden by a child class
+		InMaxPacket == 0 ? WINSOCK_MAX_PACKET : InMaxPacket,
+		InPacketOverhead == 0 ? SLIP_HEADER_SIZE : InPacketOverhead);
+
+	// Figure out IP address from the host URL
+
+	// Initialize our send bunch
+	InitSendBuffer();
+}
+
+void UWebSocketConnection::InitRemoteConnection(UNetDriver* InDriver, class FSocket* InSocket, const FURL& InURL, const class FInternetAddr& InRemoteAddr, EConnectionState InState, int32 InMaxPacket, int32 InPacketOverhead)
+{
+	InitBase(InDriver, InSocket, InURL, InState,
+		// Use the default packet size/overhead unless overridden by a child class
+		InMaxPacket == 0 ? WINSOCK_MAX_PACKET : InMaxPacket,
+		InPacketOverhead == 0 ? SLIP_HEADER_SIZE : InPacketOverhead);
+
+	// Initialize our send bunch
+	InitSendBuffer();
+
+	// This is for a client that needs to log in, setup ClientLoginState and ExpectedClientLoginMsgType to reflect that
+	SetClientLoginState(EClientLoginState::LoggingIn);
+	SetExpectedClientLoginMsgType(NMT_Hello);
+}
+
+void UWebSocketConnection::LowLevelSend(void* Data, int32 Count)
+{
+	int32 BytesSent = 0;
+	WebSocket->Send((uint8*)Data, Count);
+}
+
+FString UWebSocketConnection::LowLevelGetRemoteAddress(bool bAppendPort)
+{
+	return WebSocket->RemoteEndPoint();
+}
+
+FString UWebSocketConnection::LowLevelDescribe()
+{
+	return FString::Printf
+		(
+		TEXT(" remote=%s local=%s state: %s"),
+		*WebSocket->RemoteEndPoint(),
+		*WebSocket->LocalEndPoint(),
+		State == USOCK_Pending ? TEXT("Pending")
+		: State == USOCK_Open ? TEXT("Open")
+		: State == USOCK_Closed ? TEXT("Closed")
+		: TEXT("Invalid")
+		);
+}
+
+
+void UWebSocketConnection::SetWebSocket(FWebSocket* InWebSocket)
+{
+	WebSocket = InWebSocket; 
+}
+
+FWebSocket* UWebSocketConnection::GetWebSocket()
+{
+	return WebSocket;
+}
+
+
+void UWebSocketConnection::Tick()
+{
+	UNetConnection::Tick();
+	WebSocket->Tick();
+}
+
+void UWebSocketConnection::FinishDestroy()
+{
+	Super::FinishDestroy();
+	delete WebSocket; 
+	WebSocket = NULL;
+
+}
+
