Fix final messages not being flushed by web socket. Fixes UE-15039.

diff --git a/Engine/Plugins/Experimental/HTML5Networking/Source/HTML5Networking/Private/WebSocket.cpp b/Engine/Plugins/Experimental/HTML5Networking/Source/HTML5Networking/Private/WebSocket.cpp
--- a/Engine/Plugins/Experimental/HTML5Networking/Source/HTML5Networking/Private/WebSocket.cpp
+++ b/Engine/Plugins/Experimental/HTML5Networking/Source/HTML5Networking/Private/WebSocket.cpp
@@ -126,9 +126,10 @@ FWebSocket::FWebSocket(
 }
 
 FWebSocket::FWebSocket(WebSocketInternalContext* InContext, WebSocketInternal* InWsi)
-:	Context(InContext), 
-	Wsi(InWsi),
-	IsServerSide(true)
+	: Context(InContext)
+	, Wsi(InWsi)
+	, IsServerSide(true)
+	, Protocols(nullptr)
 {
 }
 
@@ -188,7 +189,11 @@ FString FWebSocket::LocalEndPoint()
 
 void FWebSocket::Tick()
 {
+	HandlePacket();
+}
 
+void FWebSocket::HandlePacket()
+{
 #if !PLATFORM_HTML5
 	{
 		libwebsocket_service(Context, 0);
@@ -227,6 +232,20 @@ void FWebSocket::Tick()
 #endif 
 }
 
+void FWebSocket::Flush()
+{
+	while (OutgoingBuffer.Num() > 0 && !IsServerSide)
+	{
+#if !PLATFORM_HTML5
+		if (Protocols)
+			libwebsocket_callback_on_writable_all_protocol(&Protocols[0]);
+		else
+			libwebsocket_callback_on_writable(Context, Wsi);
+#endif
+		HandlePacket();
+	};
+}
+
 void FWebSocket::SetConnectedCallBack(FWebsocketInfoCallBack CallBack)
 {
 	ConnectedCallBack = CallBack; 
@@ -247,7 +266,7 @@ void FWebSocket::OnRawRecieve(void* Data, uint32 Size)
 		uint32 BytesToBeRead = *(uint32*)RecievedBuffer.GetData();
 		if (BytesToBeRead <= ((uint32)RecievedBuffer.Num() - sizeof(uint32)))
 		{
-			RecievedCallBack.Execute((void*)((uint8*)RecievedBuffer.GetData() + sizeof(uint32)), BytesToBeRead);
+			RecievedCallBack.ExecuteIfBound((void*)((uint8*)RecievedBuffer.GetData() + sizeof(uint32)), BytesToBeRead);
 			RecievedBuffer.RemoveAt(0, sizeof(uint32) + BytesToBeRead );
 		}
 		else
@@ -264,10 +283,11 @@ void FWebSocket::OnRawRecieve(void* Data, uint32 Size)
 
 	uint32 DataToBeRead = 0;*(uint32*)Buffer;
 
-	if (Result == -1) 
+	if (Result != sizeof(uint32)) 
 	{
 		UE_LOG(LogHTML5Networking, Log, TEXT("Read message size failed!"));
 		this->ErrorCallBack.ExecuteIfBound(); 
+		return;
 	}
 	else
 	{
@@ -289,7 +309,7 @@ void FWebSocket::OnRawRecieve(void* Data, uint32 Size)
 	{
 		UE_LOG(LogHTML5Networking, Log, TEXT("Read %d bytes and Executing."), DataToBeRead);
 		check(DataToBeRead == Result);
-		RecievedCallBack.Execute(Buffer, DataToBeRead);
+		RecievedCallBack.ExecuteIfBound(Buffer, DataToBeRead);
 	}
 
 #endif 
@@ -306,29 +326,43 @@ void FWebSocket::OnRawWebSocketWritable(WebSocketInternal* wsi)
 #if !PLATFORM_HTML5_BROWSER
 
 	uint32 TotalDataSize = Packet.Num() - LWS_SEND_BUFFER_PRE_PADDING - LWS_SEND_BUFFER_POST_PADDING;
+	uint32 DataToSend = TotalDataSize;
+	while (DataToSend)
+	{
+		int Sent = libwebsocket_write(Wsi, Packet.GetData() + LWS_SEND_BUFFER_PRE_PADDING + (DataToSend-TotalDataSize), DataToSend, (libwebsocket_write_protocol)LWS_WRITE_BINARY);
+		if (Sent < 0)
+		{
+			ErrorCallBack.ExecuteIfBound();
+			return;
+		}
+		if ((uint32)Sent < DataToSend)
+		{
+			UE_LOG(LogHTML5Networking, Warning, TEXT("Could not write all '%d' bytes to socket"), DataToSend);
+		}
+		DataToSend-=Sent;
+	}
 
-	int Sent = libwebsocket_write(Wsi, Packet.GetData() + LWS_SEND_BUFFER_PRE_PADDING, TotalDataSize, (libwebsocket_write_protocol)LWS_WRITE_BINARY);
-
-	check(Sent == TotalDataSize); // if this fires we need a slightly more complicated buffering mechanism. 
 	check(Wsi == wsi);
 
 #endif
 
 #if  PLATFORM_HTML5_BROWSER
 
-	// send actual data in one go. 
-	int Result = send(SockFd, (uint32*)Packet.GetData(),Packet.Num(), 0);
-
-	if (Result == -1)
+	uint32 TotalDataSize = Packet.Num();
+	uint32 DataToSend = TotalDataSize;
+	while (DataToSend)
 	{
-		// we are caught with our pants down. fail. 
-		UE_LOG(LogHTML5Networking, Error, TEXT("Could not write %d bytes"), Packet.Num());
-		ErrorCallBack.ExecuteIfBound(); 
-	}
-	else
-	{
-		check(Result == Packet.Num());
-		UE_LOG(LogHTML5Networking, Log, TEXT("Wrote %d bytes"), Packet.Num());
+		// send actual data in one go. 
+		int Result = send(SockFd, Packet.GetData()+(DataToSend-TotalDataSize),DataToSend, 0);
+		if (Result == -1)
+		{
+			// we are caught with our pants down. fail. 
+			UE_LOG(LogHTML5Networking, Error, TEXT("Could not write %d bytes"), Packet.Num());
+			ErrorCallBack.ExecuteIfBound(); 
+			return;
+		}
+		UE_CLOG((uint32)Result < DataToSend, LogHTML5Networking, Warning, TEXT("Could not write all '%d' bytes to socket"), DataToSend);
+		DataToSend-=Result;
 	}
 	
 #endif 
@@ -340,6 +374,9 @@ void FWebSocket::OnRawWebSocketWritable(WebSocketInternal* wsi)
 
 FWebSocket::~FWebSocket()
 {
+	RecievedCallBack.Unbind();
+	Flush();
+
 #if !PLATFORM_HTML5
 	if ( !IsServerSide)
 	{
diff --git a/Engine/Plugins/Experimental/HTML5Networking/Source/HTML5Networking/Private/WebSocket.h b/Engine/Plugins/Experimental/HTML5Networking/Source/HTML5Networking/Private/WebSocket.h
--- a/Engine/Plugins/Experimental/HTML5Networking/Source/HTML5Networking/Private/WebSocket.h
+++ b/Engine/Plugins/Experimental/HTML5Networking/Source/HTML5Networking/Private/WebSocket.h
@@ -31,6 +31,8 @@ public:
 
 	/** service libwebsocket.			   */ 
 	void Tick();
+	/** service libwebsocket until outgoing buffer is empty */ 
+	void Flush();
 
 	/** Helper functions to describe end points. */
 	FString RemoteEndPoint();
@@ -38,6 +40,7 @@ public:
 
 private:
 
+	void HandlePacket();
 	void OnRawRecieve(void* Data, uint32 Size);
 	void OnRawWebSocketWritable(WebSocketInternal* wsi);
 
