Add missing nullptr assignment to avoid double delete. Fixes UE-15040.

diff --git a/Engine/Plugins/Experimental/HTML5Networking/Source/HTML5Networking/Private/WebSocketNetDriver.cpp b/Engine/Plugins/Experimental/HTML5Networking/Source/HTML5Networking/Private/WebSocketNetDriver.cpp
--- a/Engine/Plugins/Experimental/HTML5Networking/Source/HTML5Networking/Private/WebSocketNetDriver.cpp
+++ b/Engine/Plugins/Experimental/HTML5Networking/Source/HTML5Networking/Private/WebSocketNetDriver.cpp
@@ -195,6 +195,7 @@ void UWebSocketNetDriver::LowLevelDestroy()
 {
 	Super::LowLevelDestroy();
 	delete WebSocketServer; 
+	WebSocketServer = nullptr;
 }
 
 bool UWebSocketNetDriver::HandleSocketsCommand(const TCHAR* Cmd, FOutputDevice& Ar, UWorld* InWorld)
