Remove unused dependency on OnlineSubsystem from HTML5 networking module. Now that games are forking OSS plugins for their own use, this dependency causes plugin binaries to be outdated when switching between games, which causes downstream errors when build products are modified.

diff --git a/Engine/Plugins/Experimental/HTML5Networking/Source/HTML5Networking/HTML5Networking.Build.cs b/Engine/Plugins/Experimental/HTML5Networking/Source/HTML5Networking/HTML5Networking.Build.cs
--- a/Engine/Plugins/Experimental/HTML5Networking/Source/HTML5Networking/HTML5Networking.Build.cs
+++ b/Engine/Plugins/Experimental/HTML5Networking/Source/HTML5Networking/HTML5Networking.Build.cs
@@ -6,10 +6,6 @@ namespace UnrealBuildTool.Rules
 	{
 		public HTML5Networking(TargetInfo Target)
 		{
-			Definitions.Add("ONLINESUBSYSTEMUTILS_PACKAGE=1");
-
-			PrivateIncludePaths.Add("OnlineSubsystemUtils/Private");
-
             PrivateDependencyModuleNames.AddRange(
                 new string[] { 
                     "Core", 
@@ -23,8 +19,6 @@ namespace UnrealBuildTool.Rules
                     "zlib"
                 }
             );
-
-	        PublicDependencyModuleNames.Add("OnlineSubsystem");
 		}
 	}
 }
