Part 2 of cross-platform fixes for merge of //UE4/Portal-Staging

diff --git a/Engine/Plugins/Experimental/HTML5Networking/Source/HTML5Networking/Private/WebSocket.cpp b/Engine/Plugins/Experimental/HTML5Networking/Source/HTML5Networking/Private/WebSocket.cpp
--- a/Engine/Plugins/Experimental/HTML5Networking/Source/HTML5Networking/Private/WebSocket.cpp
+++ b/Engine/Plugins/Experimental/HTML5Networking/Source/HTML5Networking/Private/WebSocket.cpp
@@ -26,6 +26,8 @@
 #endif
 
 #if !PLATFORM_HTML5
+static int unreal_networking_client(struct lws *wsi, enum lws_callback_reasons reason, void *user, void *in, size_t len);
+
 static void lws_debugLogS(int level, const char *line)
 {
 	UE_LOG(LogHTML5Networking, Log, TEXT("client: %s"), ANSI_TO_TCHAR(line));
@@ -48,7 +50,7 @@ FWebSocket::FWebSocket(
 	FMemory::Memzero(Protocols, sizeof(lws_protocols) * 3);
 
 	Protocols[0].name = "binary";
-	Protocols[0].callback = FWebSocket::unreal_networking_client;
+	Protocols[0].callback = unreal_networking_client;
 	Protocols[0].per_session_data_size = 0;
 	Protocols[0].rx_buffer_size = 10 * 1024 * 1024;
 
@@ -112,8 +114,8 @@ FWebSocket::FWebSocket(
 FWebSocket::FWebSocket(WebSocketInternalContext* InContext, WebSocketInternal* InWsi)
 	: Context(InContext)
 	, Wsi(InWsi)
-	, IsServerSide(true)
 	, Protocols(nullptr)
+	, IsServerSide(true)
 {
 	int sock = lws_get_socket_fd(Wsi);
 	socklen_t len = sizeof RemoteAddr;
@@ -387,7 +389,7 @@ FWebSocket::~FWebSocket()
 }
 
 #if !PLATFORM_HTML5
-int FWebSocket::unreal_networking_client(
+static int unreal_networking_client(
 		struct lws *Wsi,
 		enum lws_callback_reasons Reason,
 		void *User,
diff --git a/Engine/Plugins/Experimental/HTML5Networking/Source/HTML5Networking/Private/WebSocket.h b/Engine/Plugins/Experimental/HTML5Networking/Source/HTML5Networking/Private/WebSocket.h
--- a/Engine/Plugins/Experimental/HTML5Networking/Source/HTML5Networking/Private/WebSocket.h
+++ b/Engine/Plugins/Experimental/HTML5Networking/Source/HTML5Networking/Private/WebSocket.h
@@ -46,7 +46,8 @@ public:
 	FString LocalEndPoint(bool bAppendPort);
 	struct sockaddr_in* GetRemoteAddr();
 
-private:
+// this was made public because of cross-platform build issues
+public:
 
 	void HandlePacket();
 	void OnRawRecieve(void* Data, uint32 Size);
@@ -80,11 +81,5 @@ private:
 	/** Server side socket or client side*/
 	bool IsServerSide;
 
-#if !PLATFORM_HTML5
-	/* libwebsocket service functions */
-	static int unreal_networking_server(struct lws *wsi, enum lws_callback_reasons reason, void *user, void *in, size_t len);
-	static int unreal_networking_client(struct lws *wsi, enum lws_callback_reasons reason, void *user, void *in, size_t len);
-#endif
-
 	friend class FWebSocketServer;
 };
diff --git a/Engine/Plugins/Experimental/HTML5Networking/Source/HTML5Networking/Private/WebSocketServer.cpp b/Engine/Plugins/Experimental/HTML5Networking/Source/HTML5Networking/Private/WebSocketServer.cpp
--- a/Engine/Plugins/Experimental/HTML5Networking/Source/HTML5Networking/Private/WebSocketServer.cpp
+++ b/Engine/Plugins/Experimental/HTML5Networking/Source/HTML5Networking/Private/WebSocketServer.cpp
@@ -20,15 +20,7 @@ struct PerSessionDataServer
 
 #if !PLATFORM_HTML5
 // real networking handler.
-static int unreal_networking_server(
-	struct libwebsocket_context *,
-	struct libwebsocket *wsi,
-	enum libwebsocket_callback_reasons reason,
-	void *user,
-	void *in,
-	size_t
-	len
-);
+static int unreal_networking_server(struct lws *wsi, enum lws_callback_reasons reason, void *user, void *in, size_t len);
 
 #if !UE_BUILD_SHIPPING
 	void libwebsocket_debugLog(int level, const char *line)
@@ -50,7 +42,7 @@ bool FWebSocketServer::Init(uint32 Port, FWebsocketClientConnectedCallBack CallB
 	FMemory::Memzero(Protocols, sizeof(lws_protocols) * 3);
 
 	Protocols[0].name = "binary";
-	Protocols[0].callback = FWebSocket::unreal_networking_server;
+	Protocols[0].callback = unreal_networking_server;
 	Protocols[0].per_session_data_size = sizeof(PerSessionDataServer);
 	Protocols[0].rx_buffer_size = 10 * 1024 * 1024;
 
@@ -133,7 +125,7 @@ FString FWebSocketServer::Info()
 
 // callback.
 #if !PLATFORM_HTML5
-int FWebSocket::unreal_networking_server
+static int unreal_networking_server
 	(
 		struct lws *Wsi,
 		enum lws_callback_reasons Reason,
diff --git a/Engine/Plugins/Experimental/HTML5Networking/Source/HTML5Networking/Private/WebSocketServer.h b/Engine/Plugins/Experimental/HTML5Networking/Source/HTML5Networking/Private/WebSocketServer.h
--- a/Engine/Plugins/Experimental/HTML5Networking/Source/HTML5Networking/Private/WebSocketServer.h
+++ b/Engine/Plugins/Experimental/HTML5Networking/Source/HTML5Networking/Private/WebSocketServer.h
@@ -21,7 +21,8 @@ public:
 	/** Describe this libwebsocket server */
 	FString Info(); 
 
-private: 
+// this was made public because of cross-platform build issues
+public: 
 
 	/** Callback for a new websocket connection to the server */
 	FWebsocketClientConnectedCallBack  ConnectedCallBack;
diff --git a/Engine/Source/Runtime/Online/SSL/Private/SslCertificateManager.cpp b/Engine/Source/Runtime/Online/SSL/Private/SslCertificateManager.cpp
--- a/Engine/Source/Runtime/Online/SSL/Private/SslCertificateManager.cpp
+++ b/Engine/Source/Runtime/Online/SSL/Private/SslCertificateManager.cpp
@@ -1,5 +1,7 @@
 // Copyright 1998-2016 Epic Games, Inc. All Rights Reserved.
 
+#if WITH_SSL
+
 #include "SslPrivatePCH.h"
 
 #include "SslCertificateManager.h"
@@ -119,3 +121,5 @@ void FSslCertificateManager::EmptyRootCertificateArray()
 	}
 	RootCertificateArray.Reset();
 }
+
+#endif // #if WITH_SSL
\ No newline at end of file
diff --git a/Engine/Source/Runtime/Online/SSL/Private/SslModule.cpp b/Engine/Source/Runtime/Online/SSL/Private/SslModule.cpp
--- a/Engine/Source/Runtime/Online/SSL/Private/SslModule.cpp
+++ b/Engine/Source/Runtime/Online/SSL/Private/SslModule.cpp
@@ -29,14 +29,18 @@ void FSslModule::StartupModule()
 {	
 	Singleton = this;
 
+#if WITH_SSL
 	CertificateManagerPtr = new FSslCertificateManager();
 	static_cast<FSslCertificateManager*>(CertificateManagerPtr)->BuildRootCertificateArray();
+#endif //#if WITH_SSL
 }
 
 void FSslModule::ShutdownModule()
 {
+#if WITH_SSL
 	static_cast<FSslCertificateManager*>(CertificateManagerPtr)->EmptyRootCertificateArray();
 	delete CertificateManagerPtr;
+#endif // #if WITH_SSL
 
 	Singleton = nullptr;
 }
diff --git a/Engine/Source/Runtime/Online/SSL/Public/SslCertificateManager.h b/Engine/Source/Runtime/Online/SSL/Public/SslCertificateManager.h
deleted file mode 100644
