Possible last fix for cross-platform errors induced by merge from //UE4/Portal-Staging

diff --git a/Engine/Plugins/Experimental/HTML5Networking/Source/HTML5Networking/Private/WebSocket.cpp b/Engine/Plugins/Experimental/HTML5Networking/Source/HTML5Networking/Private/WebSocket.cpp
--- a/Engine/Plugins/Experimental/HTML5Networking/Source/HTML5Networking/Private/WebSocket.cpp
+++ b/Engine/Plugins/Experimental/HTML5Networking/Source/HTML5Networking/Private/WebSocket.cpp
@@ -71,13 +71,10 @@ FWebSocket::FWebSocket(
 
 	check(Context);
 
-
-	Wsi = lws_client_connect_extended
-							(Context,
-							TCHAR_TO_ANSI(*ServerAddress.ToString(false)),
-							ServerAddress.GetPort(),
-							false, "/", TCHAR_TO_ANSI(*ServerAddress.ToString(false)), TCHAR_TO_ANSI(*ServerAddress.ToString(false)), Protocols[1].name, -1,this);
-
+	struct lws_client_connect_info ConnectInfo = {
+			Context, TCHAR_TO_ANSI(*ServerAddress.ToString(false)), ServerAddress.GetPort(), false, "/", TCHAR_TO_ANSI(*ServerAddress.ToString(false)), TCHAR_TO_ANSI(*ServerAddress.ToString(false)), Protocols[1].name, -1, this
+	};
+	Wsi = lws_client_connect_via_info(&ConnectInfo);
 	check(Wsi);
 
 #else // PLATFORM_HTML5_BROWSER
