Fixed WebSockets to disable IPv6 for the moment. SSL changes coming in subsequent changelist. Early merge from pre-multithreading refactor.

diff --git a/Engine/Plugins/Experimental/HTML5Networking/Source/HTML5Networking/Private/WebSocket.cpp b/Engine/Plugins/Experimental/HTML5Networking/Source/HTML5Networking/Private/WebSocket.cpp
--- a/Engine/Plugins/Experimental/HTML5Networking/Source/HTML5Networking/Private/WebSocket.cpp
+++ b/Engine/Plugins/Experimental/HTML5Networking/Source/HTML5Networking/Private/WebSocket.cpp
@@ -72,6 +72,9 @@ FWebSocket::FWebSocket(
 	Info.uid = -1;
 	Info.user = this;
 
+	//Info.options |= LWS_SERVER_OPTION_DO_SSL_GLOBAL_INIT;
+	Info.options |= LWS_SERVER_OPTION_DISABLE_IPV6;
+
 	Context = lws_create_context(&Info);
 
 	check(Context);
diff --git a/Engine/Plugins/Experimental/HTML5Networking/Source/HTML5Networking/Private/WebSocketServer.cpp b/Engine/Plugins/Experimental/HTML5Networking/Source/HTML5Networking/Private/WebSocketServer.cpp
--- a/Engine/Plugins/Experimental/HTML5Networking/Source/HTML5Networking/Private/WebSocketServer.cpp
+++ b/Engine/Plugins/Experimental/HTML5Networking/Source/HTML5Networking/Private/WebSocketServer.cpp
@@ -67,6 +67,10 @@ bool FWebSocketServer::Init(uint32 Port, FWebsocketClientConnectedCallBack CallB
 	Info.options = 0;
 	// tack on this object.
 	Info.user = this;
+
+	//Info.options |= LWS_SERVER_OPTION_DO_SSL_GLOBAL_INIT;
+	Info.options |= LWS_SERVER_OPTION_DISABLE_IPV6;
+
 	Context = lws_create_context(&Info);
 
 	if (Context == NULL)
