Fixed bad compile fix, for HTML5 and Oculus NetConnections.

diff --git a/Engine/Plugins/Experimental/HTML5Networking/Source/HTML5Networking/Private/WebsocketConnection.cpp b/Engine/Plugins/Experimental/HTML5Networking/Source/HTML5Networking/Private/WebsocketConnection.cpp
--- a/Engine/Plugins/Experimental/HTML5Networking/Source/HTML5Networking/Private/WebsocketConnection.cpp
+++ b/Engine/Plugins/Experimental/HTML5Networking/Source/HTML5Networking/Private/WebsocketConnection.cpp
@@ -69,7 +69,6 @@ void UWebSocketConnection::InitRemoteConnection(UNetDriver* InDriver, class FSoc
 void UWebSocketConnection::LowLevelSend(void* Data, int32 CountBits, FOutPacketTraits& Traits)
 {
 	const uint8* DataToSend = reinterpret_cast<uint8*>(Data);
-	uint32 CountBytes = 0;
 
 	// Process any packet modifiers
 	if (Handler.IsValid() && !Handler->GetRawSend())
@@ -80,7 +79,6 @@ void UWebSocketConnection::LowLevelSend(void* Data, int32 CountBits, FOutPacketT
 		{
 			DataToSend = ProcessedData.Data;
 			CountBits = ProcessedData.CountBits;
-			CountBytes = FMath::DivideAndRoundUp(ProcessedData.CountBits, 8);
 		}
 		else
 		{
@@ -88,6 +86,8 @@ void UWebSocketConnection::LowLevelSend(void* Data, int32 CountBits, FOutPacketT
 		}
 	}
 
+	uint32 CountBytes = FMath::DivideAndRoundUp(CountBits, 8);
+
 	if (CountBits > (MaxPacket * 8))
 	{
 		UE_LOG( LogNet, Warning, TEXT( "UWebSocketConnection::LowLevelSend: CountBytes > MaxPacketSize! Count: %i, MaxPacket: %i %s" ), CountBytes, MaxPacket, *Describe() );
@@ -96,6 +96,8 @@ void UWebSocketConnection::LowLevelSend(void* Data, int32 CountBits, FOutPacketT
 	bool bBlockSend = false;
 
 #if !UE_BUILD_SHIPPING
+	UE_LOG(LogNetTraffic, VeryVerbose, TEXT("UWebSocketConnection::LowLevelSend: Address: %s, CountBytes: %i"), *LowLevelGetRemoteAddress(true), CountBytes);
+
 	LowLevelSendDel.ExecuteIfBound((void*)DataToSend, CountBytes, bBlockSend);
 #endif
 
@@ -150,6 +152,10 @@ void UWebSocketConnection::FinishDestroy()
 
 void UWebSocketConnection::ReceivedRawPacket(void* Data,int32 Count)
 {
+#if !UE_BUILD_SHIPPING
+	UE_LOG(LogNetTraffic, VeryVerbose, TEXT("UWebSocketConnection::ReceivedRawPacket: Address: %s, Count: %i"), *LowLevelGetRemoteAddress(true), Count);
+#endif
+
 	if (Count == 0 ||   // nothing to process
 		Driver == NULL) // connection closing
 	{
diff --git a/Engine/Plugins/Online/OnlineSubsystemOculus/Source/Private/OculusNetConnection.cpp b/Engine/Plugins/Online/OnlineSubsystemOculus/Source/Private/OculusNetConnection.cpp
--- a/Engine/Plugins/Online/OnlineSubsystemOculus/Source/Private/OculusNetConnection.cpp
+++ b/Engine/Plugins/Online/OnlineSubsystemOculus/Source/Private/OculusNetConnection.cpp
@@ -88,7 +88,6 @@ void UOculusNetConnection::LowLevelSend(void* Data, int32 CountBits, FOutPacketT
 	}
 
 	const uint8* DataToSend = reinterpret_cast<uint8*>(Data);
-	uint32 CountBytes = 0;
 
 	// Process any packet modifiers
 	if (Handler.IsValid() && !Handler->GetRawSend())
@@ -99,15 +98,15 @@ void UOculusNetConnection::LowLevelSend(void* Data, int32 CountBits, FOutPacketT
 		{
 			DataToSend = ProcessedData.Data;
 			CountBits = ProcessedData.CountBits;
-			CountBytes = FMath::DivideAndRoundUp(ProcessedData.CountBits, 8);
 		}
 		else
 		{
-			CountBytes = 0;
+			CountBits = 0;
 		}
 	}
 
 	bool bBlockSend = false;
+	uint32 CountBytes = FMath::DivideAndRoundUp(CountBits, 8);
 
 #if !UE_BUILD_SHIPPING
 	LowLevelSendDel.ExecuteIfBound((void*)DataToSend, CountBytes, bBlockSend);
