WebSocket - fix timed out (and misbehaving connecting) socket crash

#rnx

diff --git a/Engine/Plugins/Experimental/HTML5Networking/Source/HTML5Networking/Private/WebSocket.cpp b/Engine/Plugins/Experimental/HTML5Networking/Source/HTML5Networking/Private/WebSocket.cpp
--- a/Engine/Plugins/Experimental/HTML5Networking/Source/HTML5Networking/Private/WebSocket.cpp
+++ b/Engine/Plugins/Experimental/HTML5Networking/Source/HTML5Networking/Private/WebSocket.cpp
@@ -488,7 +488,7 @@ static int unreal_networking_client(
 			lws_set_timeout(Wsi, NO_PENDING_TIMEOUT, 0);
 			break;
 		}
-	case LWS_CALLBACK_CLOSED:
+	case LWS_CALLBACK_CLIENT_CLOSED:
 		{
 			Socket->ErrorCallBack.ExecuteIfBound();
 			return -1;
diff --git a/Engine/Plugins/Experimental/HTML5Networking/Source/HTML5Networking/Private/WebSocketServer.cpp b/Engine/Plugins/Experimental/HTML5Networking/Source/HTML5Networking/Private/WebSocketServer.cpp
--- a/Engine/Plugins/Experimental/HTML5Networking/Source/HTML5Networking/Private/WebSocketServer.cpp
+++ b/Engine/Plugins/Experimental/HTML5Networking/Source/HTML5Networking/Private/WebSocketServer.cpp
@@ -78,13 +78,10 @@ bool FWebSocketServer::Init(uint32 Port, FWebsocketClientConnectedCallBack CallB
 		ServerPort = 0;
 		delete Protocols;
 		Protocols = NULL;
-		IsAlive = false;
 		return false; // couldn't create a server.
 	}
 
 	ConnectedCallBack = CallBack;
-	IsAlive = true;
-
 #endif
 	return true;
 }
@@ -92,11 +89,8 @@ bool FWebSocketServer::Init(uint32 Port, FWebsocketClientConnectedCallBack CallB
 bool FWebSocketServer::Tick()
 {
 #if USE_LIBWEBSOCKET
-	if (IsAlive)
-	{
-		lws_service(Context, 0);
-		lws_callback_on_writable_all_protocol(Context, &Protocols[0]);
-	}
+	lws_service(Context, 0);
+	lws_callback_on_writable_all_protocol(Context, &Protocols[0]);
 #endif
 	return true;
 }
@@ -115,8 +109,6 @@ FWebSocketServer::~FWebSocketServer()
 
 	 delete Protocols;
 	 Protocols = NULL;
-
-	 IsAlive = false;
 #endif
 }
 
@@ -143,7 +135,7 @@ static int unreal_networking_server
 	struct lws_context *Context = lws_get_context(Wsi);
 	PerSessionDataServer* BufferInfo = (PerSessionDataServer*)User;
 	FWebSocketServer* Server = (FWebSocketServer*)lws_context_user(Context);
-	if (!Server->IsAlive)
+	if (!Context)
 	{
 		return 0;
 	}
@@ -180,9 +172,6 @@ static int unreal_networking_server
 		case LWS_CALLBACK_PROTOCOL_DESTROY:
 		case LWS_CALLBACK_CLOSED:
 		case LWS_CALLBACK_CLOSED_HTTP:
-			{
-				Server->IsAlive = false;
-			}
 			break;
 	}
 
diff --git a/Engine/Plugins/Experimental/HTML5Networking/Source/HTML5Networking/Private/WebSocketServer.h b/Engine/Plugins/Experimental/HTML5Networking/Source/HTML5Networking/Private/WebSocketServer.h
--- a/Engine/Plugins/Experimental/HTML5Networking/Source/HTML5Networking/Private/WebSocketServer.h
+++ b/Engine/Plugins/Experimental/HTML5Networking/Source/HTML5Networking/Private/WebSocketServer.h
@@ -35,7 +35,6 @@ public:
 
 	friend class FWebSocket;
 	uint32 ServerPort;
-	bool IsAlive;
 };
 
 
