WebSocket - fix timed out (and misbehaving connecting) socket crash

#rnx

diff --git a/Engine/Plugins/Experimental/HTML5Networking/Source/HTML5Networking/Private/WebSocketServer.cpp b/Engine/Plugins/Experimental/HTML5Networking/Source/HTML5Networking/Private/WebSocketServer.cpp
--- a/Engine/Plugins/Experimental/HTML5Networking/Source/HTML5Networking/Private/WebSocketServer.cpp
+++ b/Engine/Plugins/Experimental/HTML5Networking/Source/HTML5Networking/Private/WebSocketServer.cpp
@@ -135,10 +135,6 @@ static int unreal_networking_server
 	struct lws_context *Context = lws_get_context(Wsi);
 	PerSessionDataServer* BufferInfo = (PerSessionDataServer*)User;
 	FWebSocketServer* Server = (FWebSocketServer*)lws_context_user(Context);
-	if (!Context)
-	{
-		return 0;
-	}
 
 	switch (Reason)
 	{
@@ -158,10 +154,11 @@ static int unreal_networking_server
 			break;
 
 		case LWS_CALLBACK_SERVER_WRITEABLE:
+			if (BufferInfo->Socket->Context == Context) // UE-68340 -- bandaid until this file is removed in favor of using LwsWebSocketsManager.cpp & LwsWebSocket.cpp
 			{
 				BufferInfo->Socket->OnRawWebSocketWritable(Wsi);
-				lws_set_timeout(Wsi, NO_PENDING_TIMEOUT, 0);
 			}
+			lws_set_timeout(Wsi, NO_PENDING_TIMEOUT, 0);
 			break;
 		case LWS_CALLBACK_CLIENT_CONNECTION_ERROR:
 			{
diff --git a/Engine/Plugins/Experimental/HTML5Networking/Source/HTML5Networking/Private/WebsocketConnection.cpp b/Engine/Plugins/Experimental/HTML5Networking/Source/HTML5Networking/Private/WebsocketConnection.cpp
--- a/Engine/Plugins/Experimental/HTML5Networking/Source/HTML5Networking/Private/WebsocketConnection.cpp
+++ b/Engine/Plugins/Experimental/HTML5Networking/Source/HTML5Networking/Private/WebsocketConnection.cpp
@@ -145,9 +145,12 @@ void UWebSocketConnection::Tick()
 void UWebSocketConnection::FinishDestroy()
 {
 	Super::FinishDestroy();
-	delete WebSocket; 
-	WebSocket = NULL;
-
+	if (WebSocket)
+	{
+		WebSocket->Context = NULL; // UE-68340
+		delete WebSocket; 
+		WebSocket = NULL;
+	}
 }
 
 void UWebSocketConnection::ReceivedRawPacket(void* Data,int32 Count)
