WebSocket - fix timed out (and misbehaving connecting) socket crash

fix HTML5 build errors

#rnx

diff --git a/Engine/Plugins/Experimental/HTML5Networking/Source/HTML5Networking/Private/WebsocketConnection.cpp b/Engine/Plugins/Experimental/HTML5Networking/Source/HTML5Networking/Private/WebsocketConnection.cpp
--- a/Engine/Plugins/Experimental/HTML5Networking/Source/HTML5Networking/Private/WebsocketConnection.cpp
+++ b/Engine/Plugins/Experimental/HTML5Networking/Source/HTML5Networking/Private/WebsocketConnection.cpp
@@ -147,7 +147,9 @@ void UWebSocketConnection::FinishDestroy()
 	Super::FinishDestroy();
 	if (WebSocket)
 	{
+#if USE_LIBWEBSOCKET
 		WebSocket->Context = NULL; // UE-68340
+#endif
 		delete WebSocket; 
 		WebSocket = NULL;
 	}
