RE-submit HTML5 - fix websocket net driver connection spam crash and disconnect during handshaking crash

[RN]

diff --git a/Engine/Plugins/Experimental/HTML5Networking/Source/HTML5Networking/Private/WebSocketServer.cpp b/Engine/Plugins/Experimental/HTML5Networking/Source/HTML5Networking/Private/WebSocketServer.cpp
--- a/Engine/Plugins/Experimental/HTML5Networking/Source/HTML5Networking/Private/WebSocketServer.cpp
+++ b/Engine/Plugins/Experimental/HTML5Networking/Source/HTML5Networking/Private/WebSocketServer.cpp
@@ -147,10 +147,11 @@ static int unreal_networking_server
 			break;
 
 		case LWS_CALLBACK_RECEIVE:
+			if (BufferInfo->Socket->Context == Context) // UE-74107 -- bandaid until this file is removed in favor of using LwsWebSocketsManager.cpp & LwsWebSocket.cpp
 			{
 				BufferInfo->Socket->OnRawRecieve(In, Len);
-				lws_set_timeout(Wsi, NO_PENDING_TIMEOUT, 0);
 			}
+			lws_set_timeout(Wsi, NO_PENDING_TIMEOUT, 0);
 			break;
 
 		case LWS_CALLBACK_SERVER_WRITEABLE:
