HTML5 - add clear instructions in plugin description box to remove all existing NetDriverDefinitions when using WebSocketNetDriver

#rnx

diff --git a/Engine/Plugins/Experimental/HTML5Networking/HTML5Networking.uplugin b/Engine/Plugins/Experimental/HTML5Networking/HTML5Networking.uplugin
--- a/Engine/Plugins/Experimental/HTML5Networking/HTML5Networking.uplugin
+++ b/Engine/Plugins/Experimental/HTML5Networking/HTML5Networking.uplugin
@@ -2,8 +2,8 @@
 	"FileVersion" : 3,
 	"Version" : 1,
 	"VersionName" : "1.0",
-	"FriendlyName" : "Experimental HTML5 Networking Plugin",
-	"Description" : "HTML5 Networking",
+	"FriendlyName" : "Experimental WebSocket Networking Plugin",
+	"Description" : "WebSocket Networking - NOTE: MUST disable all other existing NetDriverDefinitions in order to use WebSocketNetDriver.",
 	"Category" : "Misc",
 	"CreatedBy" : "Epic Games, Inc.",
 	"CreatedByURL" : "http://epicgames.com",
@@ -25,6 +25,7 @@
 				"Mac",
 				"Win32",
 				"Win64",
+				"Linux",
 				"HTML5"
 			]
 		}
