HTML5 - add more instructions to plugin description box

when using websocket net driver (not only to remove all other existing net driver in use, but) to ALSO remove all PackHandlerComponents not supported byt HTML5 (e.g. SteamAuthComponentModuleInterface)

note: this is not the desired solution (but, is better than the "brittle" bit stuffing/slurping solution) - but will work as a stop gap measure until netwokring comes up with the robust (cpp) solution (automatic PacketHandlerComponents disabling when respective net driver is not in use)

[RN]

diff --git a/Engine/Plugins/Experimental/HTML5Networking/HTML5Networking.uplugin b/Engine/Plugins/Experimental/HTML5Networking/HTML5Networking.uplugin
--- a/Engine/Plugins/Experimental/HTML5Networking/HTML5Networking.uplugin
+++ b/Engine/Plugins/Experimental/HTML5Networking/HTML5Networking.uplugin
@@ -3,7 +3,7 @@
 	"Version" : 1,
 	"VersionName" : "1.0",
 	"FriendlyName" : "Experimental WebSocket Networking Plugin",
-	"Description" : "WebSocket Networking - NOTE: MUST disable all other existing NetDriverDefinitions in order to use WebSocketNetDriver.",
+	"Description" : "WebSocket Networking - NOTE: MUST disable all other existing NetDriverDefinitions in order to use WebSocketNetDriver. ALSO: MUST disable all PackHandlerComponents not supported by HTML5 (e.g. SteamAuthComponentModuleInterface)",
 	"Category" : "Misc",
 	"CreatedBy" : "Epic Games, Inc.",
 	"CreatedByURL" : "http://epicgames.com",
